AssetDefinition - Converting AssetTypeActions_Struct to an AssetDefinition.

diff --git a/Engine/Plugins/Editor/EngineAssetDefinitions/Source/Private/AssetDefinition_Texture.cpp b/Engine/Plugins/Editor/EngineAssetDefinitions/Source/Private/AssetDefinition_Texture.cpp
--- a/Engine/Plugins/Editor/EngineAssetDefinitions/Source/Private/AssetDefinition_Texture.cpp
+++ b/Engine/Plugins/Editor/EngineAssetDefinitions/Source/Private/AssetDefinition_Texture.cpp
@@ -27,12 +27,10 @@
 
 EAssetCommandResult UAssetDefinition_Texture::OpenAssets(const FAssetOpenArgs& OpenArgs) const
 {
-	const EToolkitMode::Type Mode = OpenArgs.ToolkitHost.IsValid() ? EToolkitMode::WorldCentric : EToolkitMode::Standalone;
-
 	for (UTexture* Texture : OpenArgs.LoadObjects<UTexture>())
 	{
 		ITextureEditorModule* TextureEditorModule = &FModuleManager::LoadModuleChecked<ITextureEditorModule>("TextureEditor");
-		TextureEditorModule->CreateTextureEditor(Mode, OpenArgs.ToolkitHost, Texture);
+		TextureEditorModule->CreateTextureEditor(OpenArgs.GetToolkitMode(), OpenArgs.ToolkitHost, Texture);
 	}
 	
 	return EAssetCommandResult::Handled;
diff --git a/Engine/Plugins/Editor/EngineAssetDefinitions/Source/Private/AssetDefinition_UserDefinedStruct.cpp b/Engine/Plugins/Editor/EngineAssetDefinitions/Source/Private/AssetDefinition_UserDefinedStruct.cpp
new file mode 100644
--- /dev/null
+++ b/Engine/Plugins/Editor/EngineAssetDefinitions/Source/Private/AssetDefinition_UserDefinedStruct.cpp
@@ -0,0 +1,37 @@
+// Copyright Epic Games, Inc. All Rights Reserved.
+
+#include "AssetDefinition_UserDefinedStruct.h"
+
+#include "BlueprintEditorModule.h"
+#include "Kismet2/StructureEditorUtils.h"
+
+#define LOCTEXT_NAMESPACE "UAssetDefinition_UserDefinedStruct"
+
+EAssetCommandResult UAssetDefinition_UserDefinedStruct::OpenAssets(const FAssetOpenArgs& OpenArgs) const
+{
+	if (FStructureEditorUtils::UserDefinedStructEnabled())
+	{
+		FBlueprintEditorModule& BlueprintEditorModule = FModuleManager::LoadModuleChecked<FBlueprintEditorModule>( "Kismet" );
+
+		for (UUserDefinedStruct* UDStruct : OpenArgs.LoadObjects<UUserDefinedStruct>())
+		{
+			BlueprintEditorModule.CreateUserDefinedStructEditor(OpenArgs.GetToolkitMode(), OpenArgs.ToolkitHost, UDStruct);
+		}
+	}
+
+	return EAssetCommandResult::Handled;
+}
+
+FText UAssetDefinition_UserDefinedStruct::GetAssetDescription(const FAssetData& AssetData) const
+{
+	FString Description = AssetData.GetTagValueRef<FString>("Tooltip");
+	if (!Description.IsEmpty())
+	{
+		Description.ReplaceInline(TEXT("\\n"), TEXT("\n"));
+		return FText::FromString(MoveTemp(Description));
+	}
+
+	return FText::GetEmpty();
+}
+
+#undef LOCTEXT_NAMESPACE
diff --git a/Engine/Plugins/Editor/EngineAssetDefinitions/Source/Private/AssetDefinition_UserDefinedStruct.h b/Engine/Plugins/Editor/EngineAssetDefinitions/Source/Private/AssetDefinition_UserDefinedStruct.h
new file mode 100644
--- /dev/null
+++ b/Engine/Plugins/Editor/EngineAssetDefinitions/Source/Private/AssetDefinition_UserDefinedStruct.h
@@ -0,0 +1,29 @@
+// Copyright Epic Games, Inc. All Rights Reserved.
+
+#pragma once
+
+#include "CoreMinimal.h"
+#include "Engine/UserDefinedStruct.h"
+#include "AssetDefinitionDefault.h"
+
+#include "AssetDefinition_UserDefinedStruct.generated.h"
+
+UCLASS()
+class ENGINEASSETDEFINITIONS_API UAssetDefinition_UserDefinedStruct : public UAssetDefinitionDefault
+{
+	GENERATED_BODY()
+
+public:
+	// UAssetDefinition Begin
+	virtual FText GetAssetDisplayName() const override { return NSLOCTEXT("AssetTypeActions", "AssetTypeActions_Struct", "Structure"); }
+	virtual FLinearColor GetAssetColor() const override { return FLinearColor(FColor(103, 206, 218)); }
+	virtual TSoftClassPtr<UObject> GetAssetClass() const override { return UUserDefinedStruct::StaticClass(); }
+	virtual TConstArrayView<FAssetCategoryPath> GetAssetCategories() const override
+	{
+		static const auto Categories = { EAssetCategoryPaths::Blueprint };
+		return Categories;
+	}
+	virtual FText GetAssetDescription(const FAssetData& AssetData) const override;
+	virtual EAssetCommandResult OpenAssets(const FAssetOpenArgs& OpenArgs) const override;
+	// UAssetDefinition End
+};
diff --git a/Engine/Source/Developer/AssetTools/Private/AssetTools.cpp b/Engine/Source/Developer/AssetTools/Private/AssetTools.cpp
--- a/Engine/Source/Developer/AssetTools/Private/AssetTools.cpp
+++ b/Engine/Source/Developer/AssetTools/Private/AssetTools.cpp
@@ -69,7 +69,6 @@
 #include "AssetTypeActions/AssetTypeActions_CompositeDataTable.h"
 #include "AssetTypeActions/AssetTypeActions_Enum.h"
 #include "AssetTypeActions/AssetTypeActions_Class.h"
-#include "AssetTypeActions/AssetTypeActions_Struct.h"
 #include "AssetTypeActions/AssetTypeActions_Font.h"
 #include "AssetTypeActions/AssetTypeActions_FontFace.h"
 #include "AssetTypeActions/AssetTypeActions_ForceFeedbackAttenuation.h"
@@ -167,6 +166,7 @@
 #include "AssetDefinitionRegistry.h"
 #include "VirtualTexturingEditorModule.h"
 #include "Algo/AnyOf.h"
+#include "Engine/UserDefinedStruct.h"
 #include "Misc/AssetFilterData.h"
 
 #include UE_INLINE_GENERATED_CPP_BY_NAME(AssetTools)
@@ -1259,7 +1259,6 @@ UAssetToolsImpl::UAssetToolsImpl(const FObjectInitializer& ObjectInitializer)
 	RegisterAssetTypeActions(MakeShareable(new FAssetTypeActions_CompositeDataTable));
 	RegisterAssetTypeActions(MakeShareable(new FAssetTypeActions_Enum));
 	RegisterAssetTypeActions(MakeShareable(new FAssetTypeActions_Class));
-	RegisterAssetTypeActions(MakeShareable(new FAssetTypeActions_Struct));
 	RegisterAssetTypeActions(MakeShareable(new FAssetTypeActions_SceneImportData));
 	RegisterAssetTypeActions(MakeShareable(new FAssetTypeActions_Font));
 	RegisterAssetTypeActions(MakeShareable(new FAssetTypeActions_FontFace));
diff --git a/Engine/Source/Developer/AssetTools/Private/AssetTypeActions/AssetTypeActions_Struct.cpp b/Engine/Source/Developer/AssetTools/Private/AssetTypeActions/AssetTypeActions_Struct.cpp
deleted file mode 100644
diff --git a/Engine/Source/Developer/AssetTools/Private/AssetTypeActions/AssetTypeActions_Struct.h b/Engine/Source/Developer/AssetTools/Private/AssetTypeActions/AssetTypeActions_Struct.h
deleted file mode 100644
diff --git a/Engine/Source/Editor/AssetDefinition/Private/AssetDefinition.cpp b/Engine/Source/Editor/AssetDefinition/Private/AssetDefinition.cpp
--- a/Engine/Source/Editor/AssetDefinition/Private/AssetDefinition.cpp
+++ b/Engine/Source/Editor/AssetDefinition/Private/AssetDefinition.cpp
@@ -16,6 +16,8 @@ FAssetCategoryPath EAssetCategoryPaths::Misc(LOCTEXT("Miscellaneous", "Miscellan
 FAssetCategoryPath EAssetCategoryPaths::Gameplay(LOCTEXT("Gameplay", "Gameplay"));
 FAssetCategoryPath EAssetCategoryPaths::Blueprint(LOCTEXT("Blueprint", "Blueprint"));
 FAssetCategoryPath EAssetCategoryPaths::Texture(LOCTEXT("Texture", "Texture"));
+FAssetCategoryPath EAssetCategoryPaths::Foliage(LOCTEXT("Foliage", "Foliage"));
+FAssetCategoryPath EAssetCategoryPaths::Input(LOCTEXT("Input", "Input"));
 
 FAssetCategoryPath::FAssetCategoryPath(const FText InCategory)
 {
diff --git a/Engine/Source/Editor/AssetDefinition/Public/AssetDefinition.h b/Engine/Source/Editor/AssetDefinition/Public/AssetDefinition.h
--- a/Engine/Source/Editor/AssetDefinition/Public/AssetDefinition.h
+++ b/Engine/Source/Editor/AssetDefinition/Public/AssetDefinition.h
@@ -83,6 +83,8 @@ struct FAssetOpenArgs : public FAssetArgs
 {
 	EAssetOpenMethod OpenMethod;
 	TSharedPtr<IToolkitHost> ToolkitHost;
+
+	EToolkitMode::Type GetToolkitMode() const { return ToolkitHost.IsValid() ? EToolkitMode::WorldCentric : EToolkitMode::Standalone; }
 };
 
 struct FAssetActivateArgs : public FAssetArgs
@@ -215,16 +217,22 @@ private:
  */
 struct ASSETDEFINITION_API EAssetCategoryPaths
 {
+	// This category is special, "Basic" assets appear at the very top level and are not placed into any submenu.
+	// Arguably the basic category should not exist and should instead be user configurable on what they feel should be
+	// top level assets.
 	static FAssetCategoryPath Basic;
+	
 	static FAssetCategoryPath Animation;
-	static FAssetCategoryPath Material;
 	static FAssetCategoryPath Audio;
-	static FAssetCategoryPath Physics;
-	static FAssetCategoryPath UI;
-	static FAssetCategoryPath Misc;
-	static FAssetCategoryPath Gameplay;
 	static FAssetCategoryPath Blueprint;
+	static FAssetCategoryPath Foliage;
+	static FAssetCategoryPath Gameplay;
+	static FAssetCategoryPath Input;
+	static FAssetCategoryPath Material;
+	static FAssetCategoryPath Misc;
+	static FAssetCategoryPath Physics;
 	static FAssetCategoryPath Texture;
+	static FAssetCategoryPath UI;
 };
 
 struct FAssetOpenSupport
