Adding another overload for creating an asset category.

diff --git a/Engine/Source/Editor/AssetDefinition/Private/AssetDefinition.cpp b/Engine/Source/Editor/AssetDefinition/Private/AssetDefinition.cpp
--- a/Engine/Source/Editor/AssetDefinition/Private/AssetDefinition.cpp
+++ b/Engine/Source/Editor/AssetDefinition/Private/AssetDefinition.cpp
@@ -27,15 +27,18 @@ FAssetCategoryPath::FAssetCategoryPath(const FText& InCategory)
 	CategoryPath = { TPair<FName, FText>(FName(*FTextInspector::GetSourceString(InCategory)), InCategory) };
 }
 
-FAssetCategoryPath::FAssetCategoryPath(FText InCategory, FText InSubCategory)
+FAssetCategoryPath::FAssetCategoryPath(const FText& InCategory, const FText& InSubCategory)
 {
-	CategoryPath = { TPair<FName, FText>(FName(*FTextInspector::GetSourceString(InCategory)), InCategory), TPair<FName, FText>(FName(*FTextInspector::GetSourceString(InSubCategory)), InSubCategory) };
+	CategoryPath = {
+		TPair<FName, FText>(FName(*FTextInspector::GetSourceString(InCategory)), InCategory),
+		TPair<FName, FText>(FName(*FTextInspector::GetSourceString(InSubCategory)), InSubCategory)
+	};
 }
 
-FAssetCategoryPath::FAssetCategoryPath(const FAssetCategoryPath& InCategory, const FText& SubCategory)
+FAssetCategoryPath::FAssetCategoryPath(const FAssetCategoryPath& InCategory, const FText& InSubCategory)
 {
     CategoryPath.Append(InCategory.CategoryPath);
-    CategoryPath.Add(TPair<FName, FText>(FName(*FTextInspector::GetSourceString(SubCategory)), SubCategory));
+    CategoryPath.Add(TPair<FName, FText>(FName(*FTextInspector::GetSourceString(InSubCategory)), InSubCategory));
 }
 
 FAssetCategoryPath::FAssetCategoryPath(TConstArrayView<FText> InCategoryPath)
@@ -48,6 +51,9 @@ FAssetCategoryPath::FAssetCategoryPath(TConstArrayView<FText> InCategoryPath)
 	}
 }
 
+// UAssetDefinition
+//---------------------------------------------------------------------------
+
 UAssetDefinition::UAssetDefinition()
 {
 }
diff --git a/Engine/Source/Editor/AssetDefinition/Public/Misc/AssetCategoryPath.h b/Engine/Source/Editor/AssetDefinition/Public/Misc/AssetCategoryPath.h
--- a/Engine/Source/Editor/AssetDefinition/Public/Misc/AssetCategoryPath.h
+++ b/Engine/Source/Editor/AssetDefinition/Public/Misc/AssetCategoryPath.h
@@ -12,7 +12,8 @@
 struct ASSETDEFINITION_API FAssetCategoryPath
 {
 	FAssetCategoryPath(const FText& InCategory);
-	FAssetCategoryPath(const FAssetCategoryPath& InCategory, const FText& SubCategory);
+	FAssetCategoryPath(const FText& InCategory, const FText& SubCategory);
+	FAssetCategoryPath(const FAssetCategoryPath& InCategory, const FText& InSubCategory);
 	FAssetCategoryPath(TConstArrayView<FText> InCategoryPath);
 
 	FName GetCategory() const { return CategoryPath[0].Key; }
