Asset Definition - Adding a conversion guide to the comments for Asset Definition.

diff --git a/Engine/Source/Editor/AssetDefinition/Public/AssetDefinition.h b/Engine/Source/Editor/AssetDefinition/Public/AssetDefinition.h
--- a/Engine/Source/Editor/AssetDefinition/Public/AssetDefinition.h
+++ b/Engine/Source/Editor/AssetDefinition/Public/AssetDefinition.h
@@ -282,8 +282,97 @@ public:
  * Asset Definitions represent top level assets that are known to the editor.
  *
  * -- Conversion Guide --
- *
+ * Asset Definitions (UAssetDefinitionDefault) are a replacement to the Asset Actions (FAssetTypeActions_Base) system.
+ * The reasons for the replacement are multitude, but the highlights are,
+ * 
+ * Asset Definitions no longer do things like provide a GetActions function, the replacement for this is you using the
+ * new UToolMenu extension system to register actions for the assets, an example is in this document.  A lot of the
+ * APIs have been tweaked and cleaned up to make them easier to upgrade in the future, a lot of the original API for
+ * Asset Actions were added one at a time, and several of them could be combined in simpler ways.  Another benefit is
+ * soon we might be able to make the AssetTools module no longer a circular dependency.
+ * 
+ * All of this is in service to what was previously impossible.  To be able to right click on assets in the Content
+ * Browser and NOT have the asset and every asset it knew about load on right click, this previous impossible to escape
+ * byproduct slowed down working in the editor constantly because things that didn’t require opening the asset became
+ * necessary, and some assets could load *A LOT* of other content.
+ * 
+ * Unfortunately I can’t prevent people from backsliding, at least for now.  Even after fixing the APIs to not require
+ * loading, people need to be cleverer (Use Asset Tag Data) in how they provide right click options for assets.  But to
+ * help, you can run the editor with -WarnIfAssetsLoaded on the command line.  I’ve added a new utility class to the
+ * engine called, FWarnIfAssetsLoadedInScope, it causes notifications with callstacks to be popped up telling you what
+ * code is actually responsible for any asset loads within earmarked scopes that should NOT be loading assets.
+ * 
+ * Backwards Compatibility
+ * The new system is fully* backwards compatible.  Asset Definitions are registered with the old Asset Tools
+ * (IAssetTools::RegisterAssetTypeActions) this is done through the   FAssetDefinitionProxy.  Similarly, Asset Actions
+ * (FAssetTypeActions_Base) are registered with the Asset Definition Registry with an Asset Definition Proxy (UAssetDefinition_AssetTypeActionsProxy).
+ * 
+ * When converting Asset Actions to AssetDefinitions and you’re trying to understand how to map a specific function to
+ * the new system it can be helpful to look at the equivalent function in FAssetDefinitionProxy.
+ * 
+ * IMPORTANT - You are no longer allowed to register multiple Asset Definitions for the same Asset Class.  There were a
+ * very small number of people doing this to do some tricky things with filters which are no longer required.  The new
+ * system will yell at you if you do this.
+ * 
+ * 
+ * Registration
+ * Registering your Asset Definition is no longer required like it was for Asset Actions.  The UObjects are automatically
+ * registered with the new Asset Definition Registry (UAssetDefinitionRegistry).
+ * 
+ * You no longer need to register Categories for your Asset Definition, like you had to do with Asset Actions.  Your
+ * Assets categories are just an array of FAssetCategoryPath.  They have accelerator constructors to just take an FText
+ * for the main category, and the sub category (if there is one) which replaces the whole “GetSubMenus” function from
+ * Asset Actions.  The new version can go further, with multiple sub menus and categories, but the UI isn’t set up for it yet.
+ * 
+ * GetActions
+ * The function GetActions no longer exists, the replacement is to somewhere put a self registering callback to register
+ * the UToolMenu extension, but you can just put it at the bottom of your .cpp for your AssetDefinition, that is where
+ * the others are.  The template goes something like this,
+ * 
+ * // Menu Extensions
+ * //--------------------------------------------------------------------
+ * 
+ * 
+ * namespace MenuExtension_YOUR_CLASS_TYPE
+ * {
+ *   static void ExecuteActionForYourType(const FToolMenuContext& InContext)
+ *   {
+ *      const UContentBrowserAssetContextMenuContext* Context = UContentBrowserAssetContextMenuContext::FindContextWithAssets(InContext);
+ *     
+ *      for (YOUR_CLASS_TYPE* LoadedInstance : Context->LoadSelectedObjects<YOUR_CLASS_TYPE>())
+ *      {
+ *      }
+ *   }
+ *  
+ *   static FDelayedAutoRegisterHelper DelayedAutoRegister(EDelayedRegisterRunPhase::EndOfEngineInit, []{
+ *      UToolMenus::RegisterStartupCallback(FSimpleMulticastDelegate::FDelegate::CreateLambda([]()
+ *      {
+ *         FToolMenuOwnerScoped OwnerScoped(UE_MODULE_NAME);
+ *         UToolMenu* Menu = UE::ContentBrowser::ExtendToolMenu_AssetContextMenu(YOUR_CLASS_TYPE::StaticClass());
+ *     
+ *         FToolMenuSection& Section = Menu->FindOrAddSection("GetAssetActions");
+ *         Section.AddDynamicEntry(NAME_None, FNewToolMenuSectionDelegate::CreateLambda([](FToolMenuSection& InSection)
+ *         {
+ *            const TAttribute<FText> Label = LOCTEXT("YOUR_ASSET_TYPE_ExecuteActionForYourType", "Execute Action For Your Type");
+ *            const TAttribute<FText> ToolTip = LOCTEXT("YOUR_ASSET_TYPE_ExecuteActionForYourTypeTooltip", "This will execute the action for your type.");
+ *            const FSlateIcon Icon = FSlateIcon();
+ * 
+ * 
+ *            FToolUIAction UIAction;
+ *            UIAction.ExecuteAction = FToolMenuExecuteAction::CreateStatic(&ExecuteActionForYourType);
+ *            InSection.AddMenuEntry("YOUR_ASSET_TYPE_ExecuteActionForYourType", Label, ToolTip, Icon, UIAction);
+ *         ));
+ *      }));
+ *   });
+ * }
+ * 
+ * It’s very important that you do not load the asset in your CanExecuteAction callback or in this self callback, you should save that until you finally get Executed.
+ * If you’re looking for examples, there are tons you'll find by searching for “namespace MenuExtension_”.
  * 
+ * GetFilterName & BuildBackendFilter
+ * These functions have been replaced by BuildFilters. The new function is great, you can provide an array of filters that are available with this asset.  So for example,
+ * Blueprints provide a filter for Blueprint Class, but they also provide the filters for Blueprint Interface, Blueprint Macro Library and Blueprint Function library,
+ * which are all UBlueprint assets, but differ based on Asset Tag data.
  */
 UCLASS(Abstract)
 class ASSETDEFINITION_API UAssetDefinition : public UObject
