Asset Definition: Added Media Category

diff --git a/Engine/Source/Editor/AssetDefinition/Private/AssetDefinition.cpp b/Engine/Source/Editor/AssetDefinition/Private/AssetDefinition.cpp
--- a/Engine/Source/Editor/AssetDefinition/Private/AssetDefinition.cpp
+++ b/Engine/Source/Editor/AssetDefinition/Private/AssetDefinition.cpp
@@ -21,6 +21,7 @@ FAssetCategoryPath EAssetCategoryPaths::Foliage(LOCTEXT("Foliage", "Foliage"));
 FAssetCategoryPath EAssetCategoryPaths::Input(LOCTEXT("Input", "Input"));
 FAssetCategoryPath EAssetCategoryPaths::FX(LOCTEXT("FX", "FX"));
 FAssetCategoryPath EAssetCategoryPaths::Cinematics(LOCTEXT("Cinematics", "Cinematics"));
+FAssetCategoryPath EAssetCategoryPaths::Media(LOCTEXT("Media", "Media"));
 
 FAssetCategoryPath::FAssetCategoryPath(const FText& InCategory)
 {
diff --git a/Engine/Source/Editor/AssetDefinition/Public/AssetDefinition.h b/Engine/Source/Editor/AssetDefinition/Public/AssetDefinition.h
--- a/Engine/Source/Editor/AssetDefinition/Public/AssetDefinition.h
+++ b/Engine/Source/Editor/AssetDefinition/Public/AssetDefinition.h
@@ -234,6 +234,7 @@ struct ASSETDEFINITION_API EAssetCategoryPaths
 	static FAssetCategoryPath Gameplay;
 	static FAssetCategoryPath Input;
 	static FAssetCategoryPath Material;
+	static FAssetCategoryPath Media;
 	static FAssetCategoryPath Misc;
 	static FAssetCategoryPath Physics;
 	static FAssetCategoryPath Texture;
