Include all subcategories specified by UAssetDefinition in the New Content menu, instead of only using the first subcategory

#rnx

diff --git a/Engine/Source/Developer/AssetTools/Private/AssetTools.cpp b/Engine/Source/Developer/AssetTools/Private/AssetTools.cpp
--- a/Engine/Source/Developer/AssetTools/Private/AssetTools.cpp
+++ b/Engine/Source/Developer/AssetTools/Private/AssetTools.cpp
@@ -218,7 +218,7 @@ public:
 	}
 
 private:
-	TArray<FText> SubMenus;
+	mutable TArray<FText> SubMenus;
 	mutable bool SubmenusInitialized = false;
 
 public:
@@ -233,10 +233,7 @@ public:
 			
 			for (const FAssetCategoryPath& Category : AssetDefinitionPtr.Get()->GetAssetCategories())
 			{
-				if (Category.HasSubCategory())
-				{
-					const_cast<FAssetDefinitionProxy*>(this)->SubMenus.Add(Category.GetSubCategoryText());
-				}
+				Category.GetSubCategoriesText(SubMenus);
 			}
 		}
     
diff --git a/Engine/Source/Editor/AssetDefinition/Private/AssetDefinition.cpp b/Engine/Source/Editor/AssetDefinition/Private/AssetDefinition.cpp
--- a/Engine/Source/Editor/AssetDefinition/Private/AssetDefinition.cpp
+++ b/Engine/Source/Editor/AssetDefinition/Private/AssetDefinition.cpp
@@ -54,6 +54,30 @@ FAssetCategoryPath::FAssetCategoryPath(TConstArrayView<FText> InCategoryPath)
 	}
 }
 
+void FAssetCategoryPath::GetSubCategories(TArray<FName>& SubCategories) const
+{
+	if (HasSubCategory())
+	{
+		SubCategories.Reserve(SubCategories.Num() + NumSubCategories());
+		for (int32 i = 1; i < CategoryPath.Num(); i++)
+		{
+			SubCategories.Add(CategoryPath[i].Key);
+		}
+	}
+}
+
+void FAssetCategoryPath::GetSubCategoriesText(TArray<FText>& SubCategories) const
+{
+	if (HasSubCategory())
+	{
+		SubCategories.Reserve(SubCategories.Num() + NumSubCategories());
+		for (int32 i = 1; i < CategoryPath.Num(); i++)
+		{
+			SubCategories.Add(CategoryPath[i].Value);
+		}
+	}
+}
+
 // UAssetDefinition
 //---------------------------------------------------------------------------
 
diff --git a/Engine/Source/Editor/AssetDefinition/Public/Misc/AssetCategoryPath.h b/Engine/Source/Editor/AssetDefinition/Public/Misc/AssetCategoryPath.h
--- a/Engine/Source/Editor/AssetDefinition/Public/Misc/AssetCategoryPath.h
+++ b/Engine/Source/Editor/AssetDefinition/Public/Misc/AssetCategoryPath.h
@@ -20,8 +20,12 @@ struct ASSETDEFINITION_API FAssetCategoryPath
 	FText GetCategoryText() const { return CategoryPath[0].Value; }
 	
 	bool HasSubCategory() const { return CategoryPath.Num() > 1; }
+	int32 NumSubCategories() const { return HasSubCategory() ? (CategoryPath.Num() - 1) : 0; }
 	FName GetSubCategory() const { return HasSubCategory() ? CategoryPath[1].Key : NAME_None; }
 	FText GetSubCategoryText() const { return HasSubCategory() ? CategoryPath[1].Value : FText::GetEmpty(); }
+
+	void GetSubCategories(TArray<FName>& SubCategories) const;
+	void GetSubCategoriesText(TArray<FText>& SubCategories) const;
 	
 	FAssetCategoryPath operator / (const FText& SubCategory) const { return FAssetCategoryPath(*this, SubCategory); }
 	
