Make the UAssetDefintions more robust by not crashing the editor if a module is loaded late in the engine life cycle.

diff --git a/Engine/Source/Editor/AssetDefinition/Private/AssetDefinition.cpp b/Engine/Source/Editor/AssetDefinition/Private/AssetDefinition.cpp
--- a/Engine/Source/Editor/AssetDefinition/Private/AssetDefinition.cpp
+++ b/Engine/Source/Editor/AssetDefinition/Private/AssetDefinition.cpp
@@ -91,7 +91,11 @@ void UAssetDefinition::PostCDOContruct()
 
 	if (CanRegisterStatically())
 	{
-		UAssetDefinitionRegistry::Get()->RegisterAssetDefinition(this);
+		// The Registry might be null if the module owning the instanced CDO is loaded late into editor lifecycle.
+		if (UAssetDefinitionRegistry* AssetDefinitionRegistry = UAssetDefinitionRegistry::Get())
+		{
+			AssetDefinitionRegistry->RegisterAssetDefinition(this);
+		}
 	}
 }
 
