[AssetDefinition] Move default status on asset from AssetDefinitionDefault to base AssetDefinition

#FailedValidations Preflight

diff --git a/Engine/Source/Editor/AssetDefinition/Private/AssetDefinition.cpp b/Engine/Source/Editor/AssetDefinition/Private/AssetDefinition.cpp
--- a/Engine/Source/Editor/AssetDefinition/Private/AssetDefinition.cpp
+++ b/Engine/Source/Editor/AssetDefinition/Private/AssetDefinition.cpp
@@ -1,7 +1,13 @@
 // Copyright Epic Games, Inc. All Rights Reserved.
 
 #include "AssetDefinition.h"
+
+#include "AssetDefinitionAssetInfo.h"
 #include "AssetDefinitionRegistry.h"
+#include "IAssetStatusInfoProvider.h"
+#include "ISourceControlModule.h"
+#include "ISourceControlProvider.h"
+#include "ISourceControlState.h"
 #include "Misc/AssetFilterData.h"
 #include "EditorFramework/AssetImportData.h"
 
@@ -27,6 +33,114 @@ FAssetCategoryPath EAssetCategoryPaths::Cinematics(LOCTEXT("Cinematics", "Cinema
 FAssetCategoryPath EAssetCategoryPaths::Media(LOCTEXT("Media", "Media"));
 FAssetCategoryPath EAssetCategoryPaths::World(LOCTEXT("World", "World"));
 
+namespace UE::AssetDefinition::Status
+{
+	EVisibility GetDirtyStatusVisibility(const TSharedPtr<IAssetStatusInfoProvider> InAssetStatusInfoProvider)
+	{
+		EVisibility DirtyStatusVisibility = EVisibility::Collapsed;
+		if (!InAssetStatusInfoProvider.IsValid())
+		{
+			return DirtyStatusVisibility;
+		}
+
+		if (const UPackage* Package = InAssetStatusInfoProvider->FindPackage())
+		{
+			DirtyStatusVisibility = Package->IsDirty() ? EVisibility::Visible : EVisibility::Collapsed;
+		}
+		return DirtyStatusVisibility;
+	}
+
+	const FSlateBrush* GetSourceControlStatusBrush(const TSharedPtr<IAssetStatusInfoProvider> InAssetStatusInfoProvider)
+	{
+		const FSlateBrush* SourceControlBrush = FAppStyle::GetNoBrush();
+		if (!InAssetStatusInfoProvider.IsValid())
+		{
+			return SourceControlBrush;
+		}
+
+		if (ISourceControlModule::Get().IsEnabled() && ISourceControlModule::Get().GetProvider().IsAvailable())
+		{
+			const FString FileName = InAssetStatusInfoProvider->TryGetFilename();
+			if (FSourceControlStatePtr SourceControlState = ISourceControlModule::Get().GetProvider().GetState(FileName, EStateCacheUsage::Use))
+			{
+				FSlateIcon SCCIcon = SourceControlState->GetIcon();
+				if (SCCIcon.IsSet())
+				{
+					SourceControlBrush = SCCIcon.GetIcon();
+				}
+			}
+		}
+		return SourceControlBrush;
+	}
+
+	const FSlateBrush* GetSourceControlStatusOverlayBrush(const TSharedPtr<IAssetStatusInfoProvider> InAssetStatusInfoProvider)
+	{
+		const FSlateBrush* SourceControlStatusOverlay = FAppStyle::GetNoBrush();
+		if (!InAssetStatusInfoProvider.IsValid())
+		{
+			return SourceControlStatusOverlay;
+		}
+
+		if (ISourceControlModule::Get().IsEnabled() && ISourceControlModule::Get().GetProvider().IsAvailable())
+		{
+			const FString FileName = InAssetStatusInfoProvider->TryGetFilename();
+			if (FSourceControlStatePtr SourceControlState = ISourceControlModule::Get().GetProvider().GetState(FileName, EStateCacheUsage::Use))
+			{
+				FSlateIcon SCCIcon = SourceControlState->GetIcon();
+				if (SCCIcon.IsSet())
+				{
+					SourceControlStatusOverlay = SCCIcon.GetOverlayIcon();
+				}
+			}
+		}
+		return SourceControlStatusOverlay;
+	}
+
+	EVisibility GetSourceControlStatusVisibility(const TSharedPtr<IAssetStatusInfoProvider> InAssetStatusInfoProvider)
+	{
+		EVisibility SourceControlStatusVisibility = EVisibility::Collapsed;
+		if (!InAssetStatusInfoProvider.IsValid())
+		{
+			return SourceControlStatusVisibility;
+		}
+
+		if (ISourceControlModule::Get().IsEnabled() && ISourceControlModule::Get().GetProvider().IsAvailable())
+		{
+			const FString FileName = InAssetStatusInfoProvider->TryGetFilename();
+			if (FSourceControlStatePtr SourceControlState = ISourceControlModule::Get().GetProvider().GetState(FileName, EStateCacheUsage::Use))
+			{
+				FSlateIcon SCCIcon = SourceControlState->GetIcon();
+				if (SCCIcon.IsSet())
+				{
+					SourceControlStatusVisibility = EVisibility::Visible;
+				}
+			}
+		}
+		return SourceControlStatusVisibility;
+	}
+
+	FText GetSourceControlStatusDescription(const TSharedPtr<IAssetStatusInfoProvider> InAssetStatusInfoProvider)
+	{
+		const FText SourceControlDescriptionEmpty = LOCTEXT("NoStatus", "Couldn't retrieve source control status");
+		FText SourceControlDescription = FText::GetEmpty();
+		if (!InAssetStatusInfoProvider.IsValid())
+		{
+			return SourceControlDescriptionEmpty;
+		}
+
+		if (ISourceControlModule::Get().IsEnabled() && ISourceControlModule::Get().GetProvider().IsAvailable())
+		{
+			const FString FileName = InAssetStatusInfoProvider->TryGetFilename();
+			if (FSourceControlStatePtr SourceControlState = ISourceControlModule::Get().GetProvider().GetState(FileName, EStateCacheUsage::Use))
+			{
+				// Calling this instead of GetStatusText, since that will check for warnings and won't give out a tooltip for every state unlike GetDisplayTooltip
+				SourceControlDescription = SourceControlState->GetDisplayTooltip();
+			}
+		}
+		return SourceControlDescription.IsEmpty() ? SourceControlDescriptionEmpty : SourceControlDescription;
+	}
+}
+
 FAssetCategoryPath::FAssetCategoryPath(const FText& InCategory)
 {
 	CategoryPath = { TPair<FName, FText>(FName(*FTextInspector::GetSourceString(InCategory)), InCategory) };
@@ -164,6 +278,24 @@ EAssetCommandResult UAssetDefinition::GetSourceFiles(const FAssetData& InAsset,
 	return EAssetCommandResult::Unhandled;
 }
 
+void UAssetDefinition::GetAssetStatusInfo(const TSharedPtr<IAssetStatusInfoProvider>& InAssetStatusInfoProvider, TArray<FAssetDisplayInfo>& OutStatusInfo) const
+{
+	FAssetDisplayInfo DirtyStatus;
+	DirtyStatus.StatusIcon = FAppStyle::GetBrush("ContentBrowser.ContentDirty");
+	DirtyStatus.Priority = FAssetStatusPriority(EStatusSeverity::Info, 5);
+	DirtyStatus.StatusDescription = LOCTEXT("DirtyAssetTooltip", "Asset has unsaved changes");
+	DirtyStatus.IsVisible = TAttribute<EVisibility>::Create(TAttribute<EVisibility>::FGetter::CreateStatic(&UE::AssetDefinition::Status::GetDirtyStatusVisibility, InAssetStatusInfoProvider));
+	OutStatusInfo.Add(DirtyStatus);
+
+	FAssetDisplayInfo SCCStatus;
+	SCCStatus.Priority = FAssetStatusPriority(EStatusSeverity::Info, 0);
+	SCCStatus.StatusIcon = TAttribute<const FSlateBrush*>::Create(TAttribute<const FSlateBrush*>::FGetter::CreateStatic(&UE::AssetDefinition::Status::GetSourceControlStatusBrush, InAssetStatusInfoProvider));
+	SCCStatus.StatusIconOverlay = TAttribute<const FSlateBrush*>::Create(TAttribute<const FSlateBrush*>::FGetter::CreateStatic(&UE::AssetDefinition::Status::GetSourceControlStatusOverlayBrush, InAssetStatusInfoProvider));
+	SCCStatus.IsVisible = TAttribute<EVisibility>::Create(TAttribute<EVisibility>::FGetter::CreateStatic(&UE::AssetDefinition::Status::GetSourceControlStatusVisibility, InAssetStatusInfoProvider));
+	SCCStatus.StatusDescription = TAttribute<FText>::Create(TAttribute<FText>::FGetter::CreateStatic(&UE::AssetDefinition::Status::GetSourceControlStatusDescription, InAssetStatusInfoProvider));
+	OutStatusInfo.Add(SCCStatus);
+}
+
 bool UAssetDefinition::CanRegisterStatically() const
 {
 	return !GetClass()->HasAnyClassFlags(CLASS_Abstract);
diff --git a/Engine/Source/Editor/AssetDefinition/Public/AssetDefinition.h b/Engine/Source/Editor/AssetDefinition/Public/AssetDefinition.h
--- a/Engine/Source/Editor/AssetDefinition/Public/AssetDefinition.h
+++ b/Engine/Source/Editor/AssetDefinition/Public/AssetDefinition.h
@@ -612,11 +612,8 @@ public:
 		return false;
 	}
 
-	// Implementation will change in the AssetDefinitionDefault in a later update
-	virtual void GetAssetStatusInfo(const TSharedPtr<IAssetStatusInfoProvider>& InAssetStatusInfoProvider, TArray<FAssetDisplayInfo>& OutStatusInfo) const
-	{
-		return;
-	}
+	/** Get the basic status info on assets like the Dirty and SourceControl one */
+	UE_API virtual void GetAssetStatusInfo(const TSharedPtr<IAssetStatusInfoProvider>& InAssetStatusInfoProvider, TArray<FAssetDisplayInfo>& OutStatusInfo) const;
 
 	/**Optionally adds custom Asset Picker buttons to the Property Editor for the selected AssetDefinition, using the provided extension data this function returns.*/
 	virtual void GetAssetActionButtonExtensions(const FAssetData& InAssetData, TArray<FAssetButtonActionExtension>& OutExtensions) const { };
diff --git a/Engine/Source/Editor/UnrealEd/Private/AssetDefinitionDefault.cpp b/Engine/Source/Editor/UnrealEd/Private/AssetDefinitionDefault.cpp
--- a/Engine/Source/Editor/UnrealEd/Private/AssetDefinitionDefault.cpp
+++ b/Engine/Source/Editor/UnrealEd/Private/AssetDefinitionDefault.cpp
@@ -2,13 +2,10 @@
 
 #include "AssetDefinitionDefault.h"
 
-#include "AssetDefinitionAssetInfo.h"
 #include "AssetToolsModule.h"
 #include "IAssetStatusInfoProvider.h"
-#include "EditorFramework/AssetImportData.h"
 #include "EditorFramework/ThumbnailInfo.h"
 #include "ISourceControlModule.h"
-#include "ISourceControlProvider.h"
 #include "Settings/EditorLoadingSavingSettings.h"
 #include "Toolkits/SimpleAssetEditor.h"
 
@@ -16,114 +13,6 @@
 
 #define LOCTEXT_NAMESPACE "AssetDefinitionDefault"
 
-namespace UE::AssetDefinitionDefault::Status
-{
-	EVisibility GetDirtyStatusVisibility(const TSharedPtr<IAssetStatusInfoProvider> InAssetStatusInfoProvider)
-	{
-		EVisibility DirtyStatusVisibility = EVisibility::Collapsed;
-		if (!InAssetStatusInfoProvider.IsValid())
-		{
-			return DirtyStatusVisibility;
-		}
-
-		if (const UPackage* Package = InAssetStatusInfoProvider->FindPackage())
-		{
-			DirtyStatusVisibility = Package->IsDirty() ? EVisibility::Visible : EVisibility::Collapsed;
-		}
-		return DirtyStatusVisibility;
-	}
-
-	const FSlateBrush* GetSourceControlStatusBrush(const TSharedPtr<IAssetStatusInfoProvider> InAssetStatusInfoProvider)
-	{
-		const FSlateBrush* SourceControlBrush = FAppStyle::GetNoBrush();
-		if (!InAssetStatusInfoProvider.IsValid())
-		{
-			return SourceControlBrush;
-		}
-
-		if (ISourceControlModule::Get().IsEnabled() && ISourceControlModule::Get().GetProvider().IsAvailable())
-		{
-			const FString FileName = InAssetStatusInfoProvider->TryGetFilename();
-			if (FSourceControlStatePtr SourceControlState = ISourceControlModule::Get().GetProvider().GetState(FileName, EStateCacheUsage::Use))
-			{
-				FSlateIcon SCCIcon = SourceControlState->GetIcon();
-				if (SCCIcon.IsSet())
-				{
-					SourceControlBrush = SCCIcon.GetIcon();
-				}
-			}
-		}
-		return SourceControlBrush;
-	}
-
-	const FSlateBrush* GetSourceControlStatusOverlayBrush(const TSharedPtr<IAssetStatusInfoProvider> InAssetStatusInfoProvider)
-	{
-		const FSlateBrush* SourceControlStatusOverlay = FAppStyle::GetNoBrush();
-		if (!InAssetStatusInfoProvider.IsValid())
-		{
-			return SourceControlStatusOverlay;
-		}
-
-		if (ISourceControlModule::Get().IsEnabled() && ISourceControlModule::Get().GetProvider().IsAvailable())
-		{
-			const FString FileName = InAssetStatusInfoProvider->TryGetFilename();
-			if (FSourceControlStatePtr SourceControlState = ISourceControlModule::Get().GetProvider().GetState(FileName, EStateCacheUsage::Use))
-			{
-				FSlateIcon SCCIcon = SourceControlState->GetIcon();
-				if (SCCIcon.IsSet())
-				{
-					SourceControlStatusOverlay = SCCIcon.GetOverlayIcon();
-				}
-			}
-		}
-		return SourceControlStatusOverlay;
-	}
-
-	EVisibility GetSourceControlStatusVisibility(const TSharedPtr<IAssetStatusInfoProvider> InAssetStatusInfoProvider)
-	{
-		EVisibility SourceControlStatusVisibility = EVisibility::Collapsed;
-		if (!InAssetStatusInfoProvider.IsValid())
-		{
-			return SourceControlStatusVisibility;
-		}
-
-		if (ISourceControlModule::Get().IsEnabled() && ISourceControlModule::Get().GetProvider().IsAvailable())
-		{
-			const FString FileName = InAssetStatusInfoProvider->TryGetFilename();
-			if (FSourceControlStatePtr SourceControlState = ISourceControlModule::Get().GetProvider().GetState(FileName, EStateCacheUsage::Use))
-			{
-				FSlateIcon SCCIcon = SourceControlState->GetIcon();
-				if (SCCIcon.IsSet())
-				{
-					SourceControlStatusVisibility = EVisibility::Visible;
-				}
-			}
-		}
-		return SourceControlStatusVisibility;
-	}
-
-	FText GetSourceControlStatusDescription(const TSharedPtr<IAssetStatusInfoProvider> InAssetStatusInfoProvider)
-	{
-		const FText SourceControlDescriptionEmpty = LOCTEXT("NoStatus", "Couldn't retrieve source control status");
-		FText SourceControlDescription = FText::GetEmpty();
-		if (!InAssetStatusInfoProvider.IsValid())
-		{
-			return SourceControlDescriptionEmpty;
-		}
-
-		if (ISourceControlModule::Get().IsEnabled() && ISourceControlModule::Get().GetProvider().IsAvailable())
-		{
-			const FString FileName = InAssetStatusInfoProvider->TryGetFilename();
-			if (FSourceControlStatePtr SourceControlState = ISourceControlModule::Get().GetProvider().GetState(FileName, EStateCacheUsage::Use))
-			{
-				// Calling this instead of GetStatusText, since that will check for warnings and won't give out a tooltip for every state unlike GetDisplayTooltip
-				SourceControlDescription = SourceControlState->GetDisplayTooltip();
-			}
-		}
-		return SourceControlDescription.IsEmpty() ? SourceControlDescriptionEmpty : SourceControlDescription;
-	}
-}
-
 EAssetCommandResult UAssetDefinitionDefault::OpenAssets(const FAssetOpenArgs& OpenArgs) const
 {
 	if (GetAssetOpenSupport(FAssetOpenSupportArgs(OpenArgs.OpenMethod)).IsSupported)
@@ -154,24 +43,6 @@ EAssetCommandResult UAssetDefinitionDefault::PerformAssetDiff(const FAssetDiffAr
 	return EAssetCommandResult::Handled;
 }
 
-void UAssetDefinitionDefault::GetAssetStatusInfo(const TSharedPtr<IAssetStatusInfoProvider>& InAssetStatusInfoProvider, TArray<FAssetDisplayInfo>& OutStatusInfo) const
-{
-	FAssetDisplayInfo DirtyStatus;
-	DirtyStatus.StatusIcon = FAppStyle::GetBrush("ContentBrowser.ContentDirty");
-	DirtyStatus.Priority = FAssetStatusPriority(EStatusSeverity::Info, 5);
-	DirtyStatus.StatusDescription = LOCTEXT("DirtyAssetTooltip", "Asset has unsaved changes");
-	DirtyStatus.IsVisible = TAttribute<EVisibility>::Create(TAttribute<EVisibility>::FGetter::CreateStatic(&UE::AssetDefinitionDefault::Status::GetDirtyStatusVisibility, InAssetStatusInfoProvider));
-	OutStatusInfo.Add(DirtyStatus);
-
-	FAssetDisplayInfo SCCStatus;
-	SCCStatus.Priority = FAssetStatusPriority(EStatusSeverity::Info, 0);
-	SCCStatus.StatusIcon = TAttribute<const FSlateBrush*>::Create(TAttribute<const FSlateBrush*>::FGetter::CreateStatic(&UE::AssetDefinitionDefault::Status::GetSourceControlStatusBrush, InAssetStatusInfoProvider));
-	SCCStatus.StatusIconOverlay = TAttribute<const FSlateBrush*>::Create(TAttribute<const FSlateBrush*>::FGetter::CreateStatic(&UE::AssetDefinitionDefault::Status::GetSourceControlStatusOverlayBrush, InAssetStatusInfoProvider));
-	SCCStatus.IsVisible = TAttribute<EVisibility>::Create(TAttribute<EVisibility>::FGetter::CreateStatic(&UE::AssetDefinitionDefault::Status::GetSourceControlStatusVisibility, InAssetStatusInfoProvider));
-	SCCStatus.StatusDescription = TAttribute<FText>::Create(TAttribute<FText>::FGetter::CreateStatic(&UE::AssetDefinitionDefault::Status::GetSourceControlStatusDescription, InAssetStatusInfoProvider));
-	OutStatusInfo.Add(SCCStatus);
-}
-
 namespace UE::Editor
 {
 	UThumbnailInfo* FindOrCreateThumbnailInfo(UObject* AssetObject, TSubclassOf<UThumbnailInfo> ThumbnailClass)
diff --git a/Engine/Source/Editor/UnrealEd/Public/AssetDefinitionDefault.h b/Engine/Source/Editor/UnrealEd/Public/AssetDefinitionDefault.h
--- a/Engine/Source/Editor/UnrealEd/Public/AssetDefinitionDefault.h
+++ b/Engine/Source/Editor/UnrealEd/Public/AssetDefinitionDefault.h
@@ -16,7 +16,6 @@ public:
 	// UAssetDefinition Begin
 	UNREALED_API virtual EAssetCommandResult OpenAssets(const FAssetOpenArgs& OpenArgs) const override;
 	UNREALED_API virtual EAssetCommandResult PerformAssetDiff(const FAssetDiffArgs& DiffArgs) const override;
-	UNREALED_API virtual void GetAssetStatusInfo(const TSharedPtr<IAssetStatusInfoProvider>& InAssetStatusInfoProvider, TArray<FAssetDisplayInfo>& OutStatusInfo) const override;
 	// UAssetDefinition End
 };
 
