Fixed "ignoring return value of function declared with const attribute" error in ensures in AssertionMacros.h.

diff --git a/Engine/Source/Editor/AssetDefinition/Public/AssetDefinition.h b/Engine/Source/Editor/AssetDefinition/Public/AssetDefinition.h
--- a/Engine/Source/Editor/AssetDefinition/Public/AssetDefinition.h
+++ b/Engine/Source/Editor/AssetDefinition/Public/AssetDefinition.h
@@ -465,7 +465,7 @@ public:
 public:
 	
 	/** Returns the name of this type */
-	virtual FText GetAssetDisplayName() const PURE_VIRTUAL(UAssetDefinition::GetAssetDisplayName, ensureMsgf(false && "NotImplemented", TEXT("NotImplemented")); return FText(); );
+	virtual FText GetAssetDisplayName() const PURE_VIRTUAL(UAssetDefinition::GetAssetDisplayName, (void)ensureMsgf(false && "NotImplemented", TEXT("NotImplemented")); return FText(); );
 
 	/**
 	 * Returns the name of this type, but allows overriding the default on a specific instance of the asset.  This
@@ -476,10 +476,10 @@ public:
 	virtual FText GetAssetDisplayName(const FAssetData& AssetData) const { return GetAssetDisplayName(); }
 
 	/** Get the supported class of this type. */
-	virtual TSoftClassPtr<UObject> GetAssetClass() const PURE_VIRTUAL(UAssetDefinition::GetAssetClass, ensureMsgf(false && "NotImplemented", TEXT("NotImplemented")); return TSoftClassPtr<UClass>(); );
+	virtual TSoftClassPtr<UObject> GetAssetClass() const PURE_VIRTUAL(UAssetDefinition::GetAssetClass, (void)ensureMsgf(false && "NotImplemented", TEXT("NotImplemented")); return TSoftClassPtr<UClass>(); );
 
 	/** Returns the color associated with this type */
-	virtual FLinearColor GetAssetColor() const PURE_VIRTUAL(UAssetDefinition::GetAssetColor, ensureMsgf(false && "NotImplemented", TEXT("NotImplemented")); return FColor::Red; );
+	virtual FLinearColor GetAssetColor() const PURE_VIRTUAL(UAssetDefinition::GetAssetColor, (void)ensureMsgf(false && "NotImplemented", TEXT("NotImplemented")); return FColor::Red; );
 
 	/** Returns additional tooltip information for the specified asset, if it has any. */
 	virtual FText GetAssetDescription(const FAssetData& AssetData) const { return FText::GetEmpty(); }
@@ -505,7 +505,7 @@ public:
 		return FAssetOpenSupport(OpenSupportArgs.OpenMethod,OpenSupportArgs.OpenMethod == EAssetOpenMethod::Edit); 
 	}
 	
-	virtual EAssetCommandResult OpenAssets(const FAssetOpenArgs& OpenArgs) const PURE_VIRTUAL(UAssetDefinition::OpenAsset, ensureMsgf(false && "NotImplemented", TEXT("NotImplemented")); return EAssetCommandResult::Unhandled; );
+	virtual EAssetCommandResult OpenAssets(const FAssetOpenArgs& OpenArgs) const PURE_VIRTUAL(UAssetDefinition::OpenAsset, (void)ensureMsgf(false && "NotImplemented", TEXT("NotImplemented")); return EAssetCommandResult::Unhandled; );
 
 	virtual EAssetCommandResult ActivateAssets(const FAssetActivateArgs& ActivateArgs) const 
 	{
